#UE4 - Advertising support, initial implementation

diff --git a/Engine/Source/Runtime/Advertising/Advertising/Advertising.Build.cs b/Engine/Source/Runtime/Advertising/Advertising/Advertising.Build.cs
new file mode 100644
--- /dev/null
+++ b/Engine/Source/Runtime/Advertising/Advertising/Advertising.Build.cs
@@ -0,0 +1,38 @@
+// Copyright 1998-2014 Epic Games, Inc. All Rights Reserved.
+
+namespace UnrealBuildTool.Rules
+{
+	public class Advertising : ModuleRules
+	{
+		public Advertising( TargetInfo Target )
+		{
+			PublicIncludePaths.AddRange(
+				new string[] {
+					// ... add public include paths required here ...
+				}
+				);
+
+			PrivateIncludePaths.AddRange(
+				new string[] {
+					// ... add other private include paths required here ...
+				}
+				);
+
+			
+			PublicDependencyModuleNames.AddRange(
+				new string[]
+				{
+					"Core",
+                    // ... add other public dependencies that you statically link with here ...
+				}
+				);
+
+			DynamicallyLoadedModuleNames.AddRange(
+				new string[]
+				{
+					// ... add any modules that your module loads dynamically here ...
+				}
+				);
+		}
+	}
+}
\ No newline at end of file
diff --git a/Engine/Source/Runtime/Advertising/Advertising/Private/Advertising.cpp b/Engine/Source/Runtime/Advertising/Advertising/Private/Advertising.cpp
new file mode 100644
--- /dev/null
+++ b/Engine/Source/Runtime/Advertising/Advertising/Private/Advertising.cpp
@@ -0,0 +1,47 @@
+// Copyright 1998-2014 Epic Games, Inc. All Rights Reserved.
+
+#include "AdvertisingPrivatePCH.h"
+#include "Interfaces/IAdvertisingProvider.h"
+
+IMPLEMENT_MODULE( FAdvertising, Advertising );
+
+DEFINE_LOG_CATEGORY_STATIC( LogAdvertising, Display, All );
+
+FAdvertising::FAdvertising()
+{
+}
+
+FAdvertising::~FAdvertising()
+{
+}
+
+TSharedPtr< IAdvertisingProvider > FAdvertising::CreateAdvertisingProvider( const FName & ProviderName )
+{
+	// Check if we can successfully load the module.
+	if ( ProviderName != NAME_None )
+	{
+		IAdvertisingProvider * Module = FModuleManager::Get().LoadModulePtr<IAdvertisingProvider>(ProviderName);
+		if (Module != NULL)
+		{
+			UE_LOG(LogAdvertising, Log, TEXT("Creating Advertising provider %s"), *ProviderName.ToString());
+			return MakeShareable< IAdvertisingProvider >( Module );
+		}
+		else
+		{
+			UE_LOG(LogAdvertising, Warning, TEXT("Failed to find Advertising provider named %s."), *ProviderName.ToString());
+		}
+	}
+	else
+	{
+		UE_LOG(LogAdvertising, Warning, TEXT("CreateAdvertisingProvider called with a module name of None."));
+	}
+	return NULL;
+}
+
+void FAdvertising::StartupModule()
+{
+}
+
+void FAdvertising::ShutdownModule()
+{
+}
diff --git a/Engine/Source/Runtime/Advertising/Advertising/Private/AdvertisingPrivatePCH.h b/Engine/Source/Runtime/Advertising/Advertising/Private/AdvertisingPrivatePCH.h
new file mode 100644
--- /dev/null
+++ b/Engine/Source/Runtime/Advertising/Advertising/Private/AdvertisingPrivatePCH.h
@@ -0,0 +1,6 @@
+// Copyright 1998-2014 Epic Games, Inc. All Rights Reserved.
+
+#include "Advertising.h"
+
+// You should place include statements to your module's private header files here.  You only need to
+// add includes for headers that are used in most of your module's source files though.
diff --git a/Engine/Source/Runtime/Advertising/Advertising/Public/Advertising.h b/Engine/Source/Runtime/Advertising/Advertising/Public/Advertising.h
new file mode 100644
--- /dev/null
+++ b/Engine/Source/Runtime/Advertising/Advertising/Public/Advertising.h
@@ -0,0 +1,60 @@
+// Copyright 1998-2014 Epic Games, Inc. All Rights Reserved.
+
+#pragma once
+
+#include "Core.h"
+#include "ModuleManager.h"
+#include "Interfaces/IAdvertisingProvider.h"
+
+/**
+ * Advertising module interface implementation
+ */
+class FAdvertising : public IModuleInterface
+{
+	//--------------------------------------------------------------------------
+	// Module functionality
+	//--------------------------------------------------------------------------
+public:
+	FAdvertising();
+	virtual ~FAdvertising();
+
+	/**
+	 * Singleton-like access to this module's interface.  This is just for convenience!
+	 * Beware of calling this during the shutdown phase, though.  Your module might have been unloaded already.
+	 *
+	 * @return Returns singleton instance, loading the module on demand if needed
+	 */
+	static inline FAdvertising& Get()
+	{
+		return FModuleManager::LoadModuleChecked< FAdvertising >( "Advertising" );
+	}
+
+	/**
+	 * Checks to see if this module is loaded and ready.  It is only valid to call Get() if IsAvailable() returns true.
+	 *
+	 * @return True if the module is loaded and ready to use
+	 */
+	static inline bool IsAvailable()
+	{
+		return FModuleManager::Get().IsModuleLoaded( "Advertising" );
+	}
+
+	static FName GetDefaultProviderName()
+	{
+		FString ProviderName;
+		GConfig->GetString( TEXT( "Advertising" ), TEXT( "ProviderName" ), ProviderName, GEngineIni );
+		return FName( *ProviderName );
+	}
+
+	virtual TSharedPtr< IAdvertisingProvider > CreateAdvertisingProvider( const FName & ProviderName );
+
+	virtual TSharedPtr< IAdvertisingProvider > GetDefaultProvider()
+	{
+		return CreateAdvertisingProvider( GetDefaultProviderName() );
+	}
+
+private:
+	virtual void StartupModule() OVERRIDE;
+	virtual void ShutdownModule() OVERRIDE;
+};
+
diff --git a/Engine/Source/Runtime/Advertising/Advertising/Public/Interfaces/IAdvertisingProvider.h b/Engine/Source/Runtime/Advertising/Advertising/Public/Interfaces/IAdvertisingProvider.h
new file mode 100644
--- /dev/null
+++ b/Engine/Source/Runtime/Advertising/Advertising/Public/Interfaces/IAdvertisingProvider.h
@@ -0,0 +1,16 @@
+// Copyright 1998-2014 Epic Games, Inc. All Rights Reserved.
+
+#pragma once
+
+#include "Core.h"
+#include "ModuleManager.h"
+
+/** Generic interface for an advertising provider. Other modules can define more and register them with this module. */
+class IAdvertisingProvider : public IModuleInterface
+{
+public:
+	virtual void ShowAdBanner( bool bShowOnBottomOfScreen ) = 0;
+	virtual void HideAdBanner() = 0;
+	virtual void CloseAdBanner() = 0;
+};
+
diff --git a/Engine/Source/Runtime/Advertising/IOS/IOSAdvertising/IOSAdvertising.Build.cs b/Engine/Source/Runtime/Advertising/IOS/IOSAdvertising/IOSAdvertising.Build.cs
new file mode 100644
--- /dev/null
+++ b/Engine/Source/Runtime/Advertising/IOS/IOSAdvertising/IOSAdvertising.Build.cs
@@ -0,0 +1,54 @@
+// Copyright 1998-2014 Epic Games, Inc. All Rights Reserved.
+
+namespace UnrealBuildTool.Rules
+{
+	public class IOSAdvertising : ModuleRules
+	{
+		public IOSAdvertising( TargetInfo Target )
+		{
+			PublicIncludePaths.AddRange(
+				new string[] {
+					// ... add public include paths required here ...
+				}
+				);
+
+			PrivateIncludePaths.AddRange(
+				new string[] {
+                    // ... add other private include paths required here ...
+				}
+				);
+
+			
+			PublicDependencyModuleNames.AddRange(
+				new string[]
+				{
+					"Core",
+					// ... add other public dependencies that you statically link with here ...
+				}
+				);
+
+			PrivateDependencyModuleNames.AddRange(
+				new string[]
+				{
+					"Advertising",
+                    // ... add private dependencies that you statically link with here ...
+				}
+				);
+			PublicIncludePathModuleNames.Add("Advertising");
+
+            PrivateIncludePathModuleNames.AddRange(
+                new string[]
+                {
+                    // ... add any private module dependencies that should include paths
+                }
+                );
+
+			DynamicallyLoadedModuleNames.AddRange(
+				new string[]
+				{
+                    // ... add any modules that your module loads dynamically here ...
+				}
+				);
+		}
+	}
+}
\ No newline at end of file
diff --git a/Engine/Source/Runtime/Advertising/IOS/IOSAdvertising/Private/IOSAdvertising.cpp b/Engine/Source/Runtime/Advertising/IOS/IOSAdvertising/Private/IOSAdvertising.cpp
new file mode 100644
--- /dev/null
+++ b/Engine/Source/Runtime/Advertising/IOS/IOSAdvertising/Private/IOSAdvertising.cpp
@@ -0,0 +1,99 @@
+// Copyright 1998-2014 Epic Games, Inc. All Rights Reserved.
+
+#include "IOSAdvertising.h"
+
+DEFINE_LOG_CATEGORY_STATIC( LogAdvertising, Display, All );
+
+IMPLEMENT_MODULE( FIOSAdvertisingProvider, IOSAdvertising );
+
+#if 0
+#import <iAd/ADBannerView.h>
+
+@interface IOSAdvertising : UIResponder <ADBannerViewDelegate>
+	/** iAd banner view, if open */
+	@property(retain) ADBannerView* BannerView;
+
+	/**
+	* Will show an iAd on the top or bottom of screen, on top of the GL view (doesn't resize
+	* the view)
+	*
+	* @param bShowOnBottomOfScreen If true, the iAd will be shown at the bottom of the screen, top otherwise
+	*/
+	-(void)ShowAdBanner:(NSNumber*)bShowOnBottomOfScreen;
+@end
+
+@implementation IOSAdvertising
+
+@synthesize BannerView;
+
+/** TRUE if the iAd banner should be on the bottom of the screen */
+BOOL bDrawOnBottom;
+
+/** true if the user wants the banner to be displayed */
+bool bWantVisibleBanner = false;
+
+/**
+* Will show an iAd on the top or bottom of screen, on top of the GL view (doesn't resize
+* the view)
+*
+* @param bShowOnBottomOfScreen If true, the iAd will be shown at the bottom of the screen, top otherwise
+*/
+-(void)ShowAdBanner:(NSNumber*)bShowOnBottomOfScreen
+{
+	bDrawOnBottom = [bShowOnBottomOfScreen boolValue];
+	bWantVisibleBanner = true;
+
+	bool bNeedsToAddSubview = false;
+	if (self.BannerView == nil)
+	{
+		self.BannerView = [[ADBannerView alloc] initWithAdType:ADAdTypeBanner];
+		self.BannerView.delegate = self;
+		bNeedsToAddSubview = true;
+	}
+	CGRect BannerFrame = CGRectZero;
+	BannerFrame.size = [self.BannerView sizeThatFits : self.RootView.bounds.size];
+
+	if (bDrawOnBottom)
+	{
+		// move to off the bottom
+		BannerFrame.origin.y = self.RootView.bounds.size.height - BannerFrame.size.height;
+	}
+
+	self.BannerView.frame = BannerFrame;
+
+	// start out hidden, will fade in when ad loads
+	self.BannerView.hidden = YES;
+	self.BannerView.alpha = 0.0f;
+
+	if (bNeedsToAddSubview)
+	{
+		[self.RootView addSubview : self.BannerView];
+	}
+	else// if (self.BannerView.bannerLoaded)
+	{
+		[self bannerViewDidLoadAd : self.BannerView];
+	}
+}
+@end
+#endif
+
+// Just call these implementations that are still in the engine for now
+// We will eventually move all of this code to this file
+extern CORE_API void IOSShowAdBanner(bool bShowOnBottomOfScreen);
+extern CORE_API void IOSHideAdBanner();
+extern CORE_API void IOSCloseAd();
+
+void FIOSAdvertisingProvider::ShowAdBanner( bool bShowOnBottomOfScreen ) 
+{
+	//IOSShowAdBanner( bShowOnBottomOfScreen );
+}
+
+void FIOSAdvertisingProvider::HideAdBanner() 
+{
+	//IOSHideAdBanner();
+}
+
+void FIOSAdvertisingProvider::CloseAdBanner() 
+{
+	//IOSCloseAd();
+}
diff --git a/Engine/Source/Runtime/Advertising/IOS/IOSAdvertising/Private/IOSAdvertisingPrivatePCH.h b/Engine/Source/Runtime/Advertising/IOS/IOSAdvertising/Private/IOSAdvertisingPrivatePCH.h
new file mode 100644
--- /dev/null
+++ b/Engine/Source/Runtime/Advertising/IOS/IOSAdvertising/Private/IOSAdvertisingPrivatePCH.h
@@ -0,0 +1,4 @@
+// Copyright 1998-2014 Epic Games, Inc. All Rights Reserved.
+
+// You should place include statements to your module's private header files here.  You only need to
+// add includes for headers that are used in most of your module's source files though.
diff --git a/Engine/Source/Runtime/Advertising/IOS/IOSAdvertising/Public/IOSAdvertising.h b/Engine/Source/Runtime/Advertising/IOS/IOSAdvertising/Public/IOSAdvertising.h
new file mode 100644
--- /dev/null
+++ b/Engine/Source/Runtime/Advertising/IOS/IOSAdvertising/Public/IOSAdvertising.h
@@ -0,0 +1,14 @@
+// Copyright 1998-2014 Epic Games, Inc. All Rights Reserved.
+
+#pragma once
+
+#include "IAdvertisingProvider.h"
+#include "Core.h"
+
+class FIOSAdvertisingProvider : public IAdvertisingProvider
+{
+public:
+	virtual void ShowAdBanner( bool bShowOnBottomOfScreen ) OVERRIDE;
+	virtual void HideAdBanner() OVERRIDE;
+	virtual void CloseAdBanner() OVERRIDE;
+};
