#UE4 - Move the implementation of iAd code into the advertising module (still not using this code path yet though)

diff --git a/Engine/Source/Runtime/Advertising/IOS/IOSAdvertising/Private/IOSAdvertising.cpp b/Engine/Source/Runtime/Advertising/IOS/IOSAdvertising/Private/IOSAdvertising.cpp
--- a/Engine/Source/Runtime/Advertising/IOS/IOSAdvertising/Private/IOSAdvertising.cpp
+++ b/Engine/Source/Runtime/Advertising/IOS/IOSAdvertising/Private/IOSAdvertising.cpp
@@ -1,25 +1,17 @@
 // Copyright 1998-2014 Epic Games, Inc. All Rights Reserved.
 
 #include "IOSAdvertising.h"
+#include "IOSAppDelegate.h"
 
 DEFINE_LOG_CATEGORY_STATIC( LogAdvertising, Display, All );
 
 IMPLEMENT_MODULE( FIOSAdvertisingProvider, IOSAdvertising );
 
-#if 0
 #import <iAd/ADBannerView.h>
 
 @interface IOSAdvertising : UIResponder <ADBannerViewDelegate>
 	/** iAd banner view, if open */
 	@property(retain) ADBannerView* BannerView;
-
-	/**
-	* Will show an iAd on the top or bottom of screen, on top of the GL view (doesn't resize
-	* the view)
-	*
-	* @param bShowOnBottomOfScreen If true, the iAd will be shown at the bottom of the screen, top otherwise
-	*/
-	-(void)ShowAdBanner:(NSNumber*)bShowOnBottomOfScreen;
 @end
 
 @implementation IOSAdvertising
@@ -27,10 +19,16 @@ IMPLEMENT_MODULE( FIOSAdvertisingProvider, IOSAdvertising );
 @synthesize BannerView;
 
 /** TRUE if the iAd banner should be on the bottom of the screen */
-BOOL bDrawOnBottom;
+static BOOL bDrawOnBottom;
 
 /** true if the user wants the banner to be displayed */
-bool bWantVisibleBanner = false;
+static bool bWantVisibleBanner = false;
+
++ (IOSAdvertising*)GetDelegate
+{
+	static IOSAdvertising * Singleton = [[IOSAdvertising alloc] init];
+	return Singleton;
+}
 
 /**
 * Will show an iAd on the top or bottom of screen, on top of the GL view (doesn't resize
@@ -51,12 +49,12 @@ bool bWantVisibleBanner = false;
 		bNeedsToAddSubview = true;
 	}
 	CGRect BannerFrame = CGRectZero;
-	BannerFrame.size = [self.BannerView sizeThatFits : self.RootView.bounds.size];
+	BannerFrame.size = [self.BannerView sizeThatFits : [IOSAppDelegate GetDelegate].RootView.bounds.size];
 
 	if (bDrawOnBottom)
 	{
 		// move to off the bottom
-		BannerFrame.origin.y = self.RootView.bounds.size.height - BannerFrame.size.height;
+		BannerFrame.origin.y = [IOSAppDelegate GetDelegate].RootView.bounds.size.height - BannerFrame.size.height;
 	}
 
 	self.BannerView.frame = BannerFrame;
@@ -67,33 +65,130 @@ bool bWantVisibleBanner = false;
 
 	if (bNeedsToAddSubview)
 	{
-		[self.RootView addSubview : self.BannerView];
+		[[IOSAppDelegate GetDelegate].RootView addSubview : self.BannerView];
+
 	}
 	else// if (self.BannerView.bannerLoaded)
 	{
 		[self bannerViewDidLoadAd : self.BannerView];
 	}
 }
-@end
+
+-(void)bannerViewDidLoadAd:(ADBannerView*)Banner
+{
+#if !NO_LOGGING
+	NSLog(@"Ad loaded!");
+#endif
+
+	if (self.BannerView.hidden && bWantVisibleBanner)
+	{
+		self.BannerView.hidden = NO;
+		[UIView animateWithDuration:0.4f
+animations:^
+		{
+			self.BannerView.alpha = 1.0f;
+		}
+		];
+	}
+}
+
+-(void)bannerView:(ADBannerView *)Banner didFailToReceiveAdWithError : (NSError *)Error
+{
+#if !NO_LOGGING
+	NSLog(@"Ad failed to load: '%@'", Error);
 #endif
 
-// Just call these implementations that are still in the engine for now
-// We will eventually move all of this code to this file
-extern CORE_API void IOSShowAdBanner(bool bShowOnBottomOfScreen);
-extern CORE_API void IOSHideAdBanner();
-extern CORE_API void IOSCloseAd();
+	// if we get an error, hide the banner 
+	[self HideAdBanner];
+}
+
+/**
+* Callback when the user clicks on an ad
+*/
+-(BOOL)bannerViewActionShouldBegin:(ADBannerView*)Banner willLeaveApplication : (BOOL)bWillLeave
+{
+	// if we aren't about to swap out the app, tell the game to pause (or whatever)
+	if (!bWillLeave)
+	{
+		FIOSAsyncTask* AsyncTask = [[FIOSAsyncTask alloc] init];
+		AsyncTask.GameThreadCallback = ^ bool(void)
+		{
+			// tell the ad manager the user clicked on the banner
+			//@TODO: IAD:			UPlatformInterfaceBase::GetInGameAdManagerSingleton()->OnUserClickedBanner();
+
+			return true;
+		};
+		[AsyncTask FinishedTask];
+	}
+
+	return YES;
+}
+
+/**
+* Callback when an ad is closed
+*/
+-(void)bannerViewActionDidFinish:(ADBannerView*)Banner
+{
+	FIOSAsyncTask* AsyncTask = [[FIOSAsyncTask alloc] init];
+	AsyncTask.GameThreadCallback = ^ bool(void)
+	{
+		// tell ad singleton we closed the ad
+		//@TODO: IAD:		UPlatformInterfaceBase::GetInGameAdManagerSingleton()->OnUserClosedAd();
+
+		return true;
+	};
+	[AsyncTask FinishedTask];
+}
+
+/**
+* Hides the iAd banner shows with ShowAdBanner. Will force close the ad if it's open
+*/
+-(void)HideAdBanner
+{
+	// fade it out
+	if (!self.BannerView.hidden)
+	{
+		[UIView animateWithDuration:0.4f
+animations:^
+		{
+			self.BannerView.alpha = 0.0f;
+		}
+completion:^(BOOL finished) 
+		   {
+			   self.BannerView.hidden = YES;
+		   }
+		];
+	}
+}
+
+-(void)UserHideAdBanner
+{
+	bWantVisibleBanner = false;
+	[self HideAdBanner];
+}
+
+/**
+* Forces closed any displayed ad. Can lead to loss of revenue
+*/
+-(void)CloseAd
+{
+	// boot user out of the ad
+	bWantVisibleBanner = false;
+	[self.BannerView cancelBannerViewAction];
+}
+@end
 
 void FIOSAdvertisingProvider::ShowAdBanner( bool bShowOnBottomOfScreen ) 
 {
-	//IOSShowAdBanner( bShowOnBottomOfScreen );
+	[[IOSAdvertising GetDelegate] performSelectorOnMainThread:@selector(ShowAdBanner:) withObject:[NSNumber numberWithBool : bShowOnBottomOfScreen] waitUntilDone : NO];
 }
 
 void FIOSAdvertisingProvider::HideAdBanner() 
 {
-	//IOSHideAdBanner();
+	[[IOSAdvertising GetDelegate] performSelectorOnMainThread:@selector(UserHideAdBanner) withObject:nil waitUntilDone : NO];
 }
 
 void FIOSAdvertisingProvider::CloseAdBanner() 
 {
-	//IOSCloseAd();
+	[[IOSAdvertising GetDelegate] performSelectorOnMainThread:@selector(CloseAd) withObject:nil waitUntilDone : NO];
 }
