#UE4 - Return advertising interface as a pointer

diff --git a/Engine/Source/Runtime/Advertising/Advertising/Private/Advertising.cpp b/Engine/Source/Runtime/Advertising/Advertising/Private/Advertising.cpp
--- a/Engine/Source/Runtime/Advertising/Advertising/Private/Advertising.cpp
+++ b/Engine/Source/Runtime/Advertising/Advertising/Private/Advertising.cpp
@@ -15,16 +15,16 @@ FAdvertising::~FAdvertising()
 {
 }
 
-TSharedPtr< IAdvertisingProvider > FAdvertising::CreateAdvertisingProvider( const FName & ProviderName )
+IAdvertisingProvider * FAdvertising::GetAdvertisingProvider( const FName & ProviderName )
 {
 	// Check if we can successfully load the module.
 	if ( ProviderName != NAME_None )
 	{
 		IAdvertisingProvider * Module = FModuleManager::Get().LoadModulePtr<IAdvertisingProvider>(ProviderName);
-		if (Module != NULL)
+		if ( Module != NULL )
 		{
 			UE_LOG(LogAdvertising, Log, TEXT("Creating Advertising provider %s"), *ProviderName.ToString());
-			return MakeShareable< IAdvertisingProvider >( Module );
+			return Module;
 		}
 		else
 		{
@@ -33,7 +33,7 @@ TSharedPtr< IAdvertisingProvider > FAdvertising::CreateAdvertisingProvider( cons
 	}
 	else
 	{
-		UE_LOG(LogAdvertising, Warning, TEXT("CreateAdvertisingProvider called with a module name of None."));
+		UE_LOG(LogAdvertising, Warning, TEXT("GetAdvertisingProvider called with a module name of None."));
 	}
 	return NULL;
 }
diff --git a/Engine/Source/Runtime/Advertising/Advertising/Public/Advertising.h b/Engine/Source/Runtime/Advertising/Advertising/Public/Advertising.h
--- a/Engine/Source/Runtime/Advertising/Advertising/Public/Advertising.h
+++ b/Engine/Source/Runtime/Advertising/Advertising/Public/Advertising.h
@@ -42,15 +42,15 @@ public:
 	static FName GetDefaultProviderName()
 	{
 		FString ProviderName;
-		GConfig->GetString( TEXT( "Advertising" ), TEXT( "ProviderName" ), ProviderName, GEngineIni );
+		GConfig->GetString( TEXT( "Advertising" ), TEXT( "DefaultProviderName" ), ProviderName, GEngineIni );
 		return FName( *ProviderName );
 	}
 
-	virtual TSharedPtr< IAdvertisingProvider > CreateAdvertisingProvider( const FName & ProviderName );
+	virtual IAdvertisingProvider * GetAdvertisingProvider( const FName & ProviderName );
 
-	virtual TSharedPtr< IAdvertisingProvider > GetDefaultProvider()
+	virtual IAdvertisingProvider * GetDefaultProvider()
 	{
-		return CreateAdvertisingProvider( GetDefaultProviderName() );
+		return GetAdvertisingProvider( GetDefaultProviderName() );
 	}
 
 private:
