- stubbed android advertising interface

diff --git a/Engine/Config/Android/AndroidEngine.ini b/Engine/Config/Android/AndroidEngine.ini
--- a/Engine/Config/Android/AndroidEngine.ini
+++ b/Engine/Config/Android/AndroidEngine.ini
@@ -6,3 +6,6 @@ AudioDeviceModuleName=AndroidAudio
 
 [DeviceProfileManager]
 DeviceProfileSelectionModule="AndroidDeviceProfileSelector"
+
+[Advertising]
+DefaultProviderName=AndroidAdvertising
diff --git a/Engine/Source/Runtime/Advertising/Android/AndroidAdvertising/AndroidAdvertising.Build.cs b/Engine/Source/Runtime/Advertising/Android/AndroidAdvertising/AndroidAdvertising.Build.cs
new file mode 100644
--- /dev/null
+++ b/Engine/Source/Runtime/Advertising/Android/AndroidAdvertising/AndroidAdvertising.Build.cs
@@ -0,0 +1,54 @@
+// Copyright 1998-2014 Epic Games, Inc. All Rights Reserved.
+
+namespace UnrealBuildTool.Rules
+{
+	public class AndroidAdvertising : ModuleRules
+	{
+		public AndroidAdvertising( TargetInfo Target )
+		{
+			PublicIncludePaths.AddRange(
+				new string[] {
+					// ... add public include paths required here ...
+				}
+				);
+
+			PrivateIncludePaths.AddRange(
+				new string[] {
+                    // ... add other private include paths required here ...
+				}
+				);
+
+			
+			PublicDependencyModuleNames.AddRange(
+				new string[]
+				{
+					"Core",
+					// ... add other public dependencies that you statically link with here ...
+				}
+				);
+
+			PrivateDependencyModuleNames.AddRange(
+				new string[]
+				{
+					"Advertising",
+                    // ... add private dependencies that you statically link with here ...
+				}
+				);
+			PublicIncludePathModuleNames.Add("Advertising");
+
+            PrivateIncludePathModuleNames.AddRange(
+                new string[]
+                {
+                    // ... add any private module dependencies that should include paths
+                }
+                );
+
+			DynamicallyLoadedModuleNames.AddRange(
+				new string[]
+				{
+                    // ... add any modules that your module loads dynamically here ...
+				}
+				);
+		}
+	}
+}
diff --git a/Engine/Source/Runtime/Advertising/Android/AndroidAdvertising/Private/AndroidAdvertising.cpp b/Engine/Source/Runtime/Advertising/Android/AndroidAdvertising/Private/AndroidAdvertising.cpp
new file mode 100644
--- /dev/null
+++ b/Engine/Source/Runtime/Advertising/Android/AndroidAdvertising/Private/AndroidAdvertising.cpp
@@ -0,0 +1,21 @@
+// Copyright 1998-2014 Epic Games, Inc. All Rights Reserved.
+
+#include "AndroidAdvertising.h"
+
+
+DEFINE_LOG_CATEGORY_STATIC( LogAdvertising, Display, All );
+
+IMPLEMENT_MODULE( FAndroidAdvertisingProvider, AndroidAdvertising );
+
+
+void FAndroidAdvertisingProvider::ShowAdBanner( bool bShowOnBottomOfScreen ) 
+{
+}
+
+void FAndroidAdvertisingProvider::HideAdBanner() 
+{
+}
+
+void FAndroidAdvertisingProvider::CloseAdBanner() 
+{
+}
diff --git a/Engine/Source/Runtime/Advertising/Android/AndroidAdvertising/Private/AndroidAdvertisingPrivatePCH.h b/Engine/Source/Runtime/Advertising/Android/AndroidAdvertising/Private/AndroidAdvertisingPrivatePCH.h
new file mode 100644
--- /dev/null
+++ b/Engine/Source/Runtime/Advertising/Android/AndroidAdvertising/Private/AndroidAdvertisingPrivatePCH.h
@@ -0,0 +1,4 @@
+// Copyright 1998-2014 Epic Games, Inc. All Rights Reserved.
+
+// You should place include statements to your module's private header files here.  You only need to
+// add includes for headers that are used in most of your module's source files though.
diff --git a/Engine/Source/Runtime/Advertising/Android/AndroidAdvertising/Public/AndroidAdvertising.h b/Engine/Source/Runtime/Advertising/Android/AndroidAdvertising/Public/AndroidAdvertising.h
new file mode 100644
--- /dev/null
+++ b/Engine/Source/Runtime/Advertising/Android/AndroidAdvertising/Public/AndroidAdvertising.h
@@ -0,0 +1,14 @@
+// Copyright 1998-2014 Epic Games, Inc. All Rights Reserved.
+
+#pragma once
+
+#include "IAdvertisingProvider.h"
+#include "Core.h"
+
+class FAndroidAdvertisingProvider : public IAdvertisingProvider
+{
+public:
+	virtual void ShowAdBanner( bool bShowOnBottomOfScreen ) OVERRIDE;
+	virtual void HideAdBanner() OVERRIDE;
+	virtual void CloseAdBanner() OVERRIDE;
+};
diff --git a/Engine/Source/Runtime/UE4Game/UE4Game.Build.cs b/Engine/Source/Runtime/UE4Game/UE4Game.Build.cs
--- a/Engine/Source/Runtime/UE4Game/UE4Game.Build.cs
+++ b/Engine/Source/Runtime/UE4Game/UE4Game.Build.cs
@@ -15,5 +15,11 @@ public class UE4Game : ModuleRules
 			DynamicallyLoadedModuleNames.Add("OnlineSubsystemIOS");
 			DynamicallyLoadedModuleNames.Add("IOSAdvertising");
 		}
+		else if (Target.Platform == UnrealTargetPlatform.Android)
+		{
+			// @todo android: Add Android online subsystem
+			PrivateDependencyModuleNames.AddRange(new string[] { "Core", "CoreUObject", "Engine", "OnlineSubsystem", "OnlineSubsystemUtils" });
+			DynamicallyLoadedModuleNames.Add("AndroidAdvertising");
+		}
 	}
 }
diff --git a/Engine/Source/UE4Game.Target.cs b/Engine/Source/UE4Game.Target.cs
--- a/Engine/Source/UE4Game.Target.cs
+++ b/Engine/Source/UE4Game.Target.cs
@@ -53,7 +53,8 @@ public class UE4GameTarget : TargetRules
 			}
 			else if (Target.Platform == UnrealTargetPlatform.Android)
 			{
-				// @todo Rocket: Add Android online subsystem
+				// @todo android: Add Android online subsystem
+				OutExtraModuleNames.Add("AndroidAdvertising");
 			}
 		}
 	}
