fix for iAd not showing when WiFi is turned off and on
UE-9614
#ios

diff --git a/Engine/Source/Runtime/Advertising/IOS/IOSAdvertising/Private/IOSAdvertising.cpp b/Engine/Source/Runtime/Advertising/IOS/IOSAdvertising/Private/IOSAdvertising.cpp
--- a/Engine/Source/Runtime/Advertising/IOS/IOSAdvertising/Private/IOSAdvertising.cpp
+++ b/Engine/Source/Runtime/Advertising/IOS/IOSAdvertising/Private/IOSAdvertising.cpp
@@ -67,7 +67,7 @@ static bool bWantVisibleBanner = false;
 	{
 		[[IOSAppDelegate GetDelegate].RootView addSubview : self.BannerView];
 	}
-	else// if (self.BannerView.bannerLoaded)
+	else if (self.BannerView.bannerLoaded)
 	{
 		[self bannerViewDidLoadAd : self.BannerView];
 	}
