UE-17404 - AdMob multiple ad units
- Changed code to support multiple ad units and the selecting thereof via blueprints including aditional nodes

diff --git a/Engine/Plugins/Runtime/Advertising/IOSTapJoy/Source/IOSTapJoy/Private/IOSTapJoy.cpp b/Engine/Plugins/Runtime/Advertising/IOSTapJoy/Source/IOSTapJoy/Private/IOSTapJoy.cpp
--- a/Engine/Plugins/Runtime/Advertising/IOSTapJoy/Source/IOSTapJoy/Private/IOSTapJoy.cpp
+++ b/Engine/Plugins/Runtime/Advertising/IOSTapJoy/Source/IOSTapJoy/Private/IOSTapJoy.cpp
@@ -11,9 +11,10 @@ class FTapJoyProvider : public IAdvertisingProvider
 	virtual void ShutdownModule() override;
 
 	/** IAdvertisingProvider implementation */
-	virtual void ShowAdBanner( bool bShowOnBottomOfScreen ) override;
+	virtual void ShowAdBanner(bool bShowOnBottomOfScreen, int32 adID) override;
 	virtual void HideAdBanner() override;
 	virtual void CloseAdBanner() override;
+	virtual int32 GetAdIDCount() override;
 };
 
 IMPLEMENT_MODULE( FTapJoyProvider, IOSTapJoy )
@@ -207,7 +208,7 @@ void FTapJoyProvider::ShutdownModule()
 {
 }
 
-void FTapJoyProvider::ShowAdBanner( bool bShowOnBottomOfScreen )
+void FTapJoyProvider::ShowAdBanner(bool bShowOnBottomOfScreen, int32 /*adID*/)
 {
 	[[IOSTapJoy GetDelegate] performSelectorOnMainThread:@selector(ShowAdBanner:) withObject:[NSNumber numberWithBool : bShowOnBottomOfScreen] waitUntilDone : NO];
 }
@@ -221,3 +222,8 @@ void FTapJoyProvider::CloseAdBanner()
 {
 	HideAdBanner();
 }
+
+int32 FTapJoyProvider::GetAdIDCount()
+{
+	return 1;
+}
diff --git a/Engine/Source/Developer/Android/AndroidPlatformEditor/Private/AndroidPlatformEditorModule.cpp b/Engine/Source/Developer/Android/AndroidPlatformEditor/Private/AndroidPlatformEditorModule.cpp
--- a/Engine/Source/Developer/Android/AndroidPlatformEditor/Private/AndroidPlatformEditorModule.cpp
+++ b/Engine/Source/Developer/Android/AndroidPlatformEditor/Private/AndroidPlatformEditorModule.cpp
@@ -60,6 +60,9 @@ class FAndroidPlatformEditorModule
 		auto &AndroidDeviceDetection = FModuleManager::LoadModuleChecked<IAndroidDeviceDetectionModule>("AndroidDeviceDetection");
 		settings->SetDeviceDetection(AndroidDeviceDetection.GetAndroidDeviceDetection());
 		settings->UpdateTargetModulePaths();
+
+		auto platformSettings = GetMutableDefault<UAndroidRuntimeSettings>();
+		platformSettings->MigrateData();
 	}
 
 	virtual void ShutdownModule() override
diff --git a/Engine/Source/Developer/Android/AndroidPlatformEditor/Private/AndroidTargetSettingsCustomization.cpp b/Engine/Source/Developer/Android/AndroidPlatformEditor/Private/AndroidTargetSettingsCustomization.cpp
--- a/Engine/Source/Developer/Android/AndroidPlatformEditor/Private/AndroidTargetSettingsCustomization.cpp
+++ b/Engine/Source/Developer/Android/AndroidPlatformEditor/Private/AndroidTargetSettingsCustomization.cpp
@@ -185,9 +185,10 @@ void FAndroidTargetSettingsCustomization::BuildAppManifestSection(IDetailLayoutB
 		.EditCondition(SetupForGooglePlayAttribute, NULL);
 
 	TSharedRef<IPropertyHandle> AdMobAdUnitIDProperty = DetailLayout.GetProperty(GET_MEMBER_NAME_CHECKED(UAndroidRuntimeSettings, AdMobAdUnitID));
-	GooglePlayCategory.AddProperty(AdMobAdUnitIDProperty)
+	AdMobAdUnitIDProperty->MarkHiddenByCustomization();
+	/*GooglePlayCategory.AddProperty(AdMobAdUnitIDProperty)
 		.EditCondition(SetupForGooglePlayAttribute, NULL);
-
+		*/
 	TSharedRef<IPropertyHandle> GooglePlayLicenseKeyProperty = DetailLayout.GetProperty(GET_MEMBER_NAME_CHECKED(UAndroidRuntimeSettings, GooglePlayLicenseKey));
 	GooglePlayCategory.AddProperty(GooglePlayLicenseKeyProperty)
 		.EditCondition(SetupForGooglePlayAttribute, NULL);
diff --git a/Engine/Source/Runtime/Advertising/Advertising/Public/Interfaces/IAdvertisingProvider.h b/Engine/Source/Runtime/Advertising/Advertising/Public/Interfaces/IAdvertisingProvider.h
--- a/Engine/Source/Runtime/Advertising/Advertising/Public/Interfaces/IAdvertisingProvider.h
+++ b/Engine/Source/Runtime/Advertising/Advertising/Public/Interfaces/IAdvertisingProvider.h
@@ -9,8 +9,9 @@
 class IAdvertisingProvider : public IModuleInterface
 {
 public:
-	virtual void ShowAdBanner( bool bShowOnBottomOfScreen ) = 0;
+	virtual void ShowAdBanner( bool bShowOnBottomOfScreen, int32 adID ) = 0;
 	virtual void HideAdBanner() = 0;
 	virtual void CloseAdBanner() = 0;
+	virtual int32 GetAdIDCount() = 0;
 };
 
diff --git a/Engine/Source/Runtime/Advertising/Android/AndroidAdvertising/Private/AndroidAdvertising.cpp b/Engine/Source/Runtime/Advertising/Android/AndroidAdvertising/Private/AndroidAdvertising.cpp
--- a/Engine/Source/Runtime/Advertising/Android/AndroidAdvertising/Private/AndroidAdvertising.cpp
+++ b/Engine/Source/Runtime/Advertising/Android/AndroidAdvertising/Private/AndroidAdvertising.cpp
@@ -8,16 +8,19 @@ DEFINE_LOG_CATEGORY_STATIC( LogAdvertising, Display, All );
 IMPLEMENT_MODULE( FAndroidAdvertisingProvider, AndroidAdvertising );
 
 
-void FAndroidAdvertisingProvider::ShowAdBanner( bool bShowOnBottomOfScreen ) 
+void FAndroidAdvertisingProvider::ShowAdBanner(bool bShowOnBottomOfScreen, int32 adID)
 {
 	FString AdUnitID;
 	bool found = GConfig->GetString(TEXT("/Script/AndroidRuntimeSettings.AndroidRuntimeSettings"), TEXT("AdMobAdUnitID"), AdUnitID, GEngineIni);
 
-	if (!found || AdUnitID.IsEmpty())
+	TArray<FString> AdUnitIDs;
+	int32 count = GConfig->GetArray(TEXT("/Script/AndroidRuntimeSettings.AndroidRuntimeSettings"), TEXT("AdMobAdUnitID"), AdUnitIDs, GEngineIni);
+	
+	if (count == 0 || adID >= count)
 		return;
 
 	extern void AndroidThunkCpp_ShowAdBanner(const FString&, bool);
-	AndroidThunkCpp_ShowAdBanner(AdUnitID, bShowOnBottomOfScreen);
+	AndroidThunkCpp_ShowAdBanner(AdUnitIDs[adID], bShowOnBottomOfScreen);
 }
 
 void FAndroidAdvertisingProvider::HideAdBanner() 
@@ -31,3 +34,11 @@ void FAndroidAdvertisingProvider::CloseAdBanner()
 	extern void AndroidThunkCpp_CloseAdBanner();
 	AndroidThunkCpp_CloseAdBanner();
 }
+
+int32 FAndroidAdvertisingProvider::GetAdIDCount()
+{
+	TArray<FString> AdUnitIDs;
+	int32 count = GConfig->GetArray(TEXT("/Script/AndroidRuntimeSettings.AndroidRuntimeSettings"), TEXT("AdMobAdUnitID"), AdUnitIDs, GEngineIni);
+
+	return count;
+}
\ No newline at end of file
diff --git a/Engine/Source/Runtime/Advertising/Android/AndroidAdvertising/Public/AndroidAdvertising.h b/Engine/Source/Runtime/Advertising/Android/AndroidAdvertising/Public/AndroidAdvertising.h
--- a/Engine/Source/Runtime/Advertising/Android/AndroidAdvertising/Public/AndroidAdvertising.h
+++ b/Engine/Source/Runtime/Advertising/Android/AndroidAdvertising/Public/AndroidAdvertising.h
@@ -8,7 +8,8 @@
 class FAndroidAdvertisingProvider : public IAdvertisingProvider
 {
 public:
-	virtual void ShowAdBanner( bool bShowOnBottomOfScreen ) override;
+	virtual void ShowAdBanner(bool bShowOnBottomOfScreen, int32 adID) override;
 	virtual void HideAdBanner() override;
 	virtual void CloseAdBanner() override;
+	virtual int32 GetAdIDCount() override;
 };
diff --git a/Engine/Source/Runtime/Advertising/IOS/IOSAdvertising/Private/IOSAdvertising.cpp b/Engine/Source/Runtime/Advertising/IOS/IOSAdvertising/Private/IOSAdvertising.cpp
--- a/Engine/Source/Runtime/Advertising/IOS/IOSAdvertising/Private/IOSAdvertising.cpp
+++ b/Engine/Source/Runtime/Advertising/IOS/IOSAdvertising/Private/IOSAdvertising.cpp
@@ -177,7 +177,7 @@ completion:^(BOOL finished)
 }
 @end
 
-void FIOSAdvertisingProvider::ShowAdBanner( bool bShowOnBottomOfScreen ) 
+void FIOSAdvertisingProvider::ShowAdBanner( bool bShowOnBottomOfScreen, int32 /*AdID*/ ) 
 {
 	[[IOSAdvertising GetDelegate] performSelectorOnMainThread:@selector(ShowAdBanner:) withObject:[NSNumber numberWithBool : bShowOnBottomOfScreen] waitUntilDone : NO];
 }
@@ -191,3 +191,8 @@ void FIOSAdvertisingProvider::CloseAdBanner()
 {
 	[[IOSAdvertising GetDelegate] performSelectorOnMainThread:@selector(CloseAd) withObject:nil waitUntilDone : NO];
 }
+
+int32 FIOSAdvertisingProvider::GetAdIDCount()
+{
+	return 1;
+}
\ No newline at end of file
diff --git a/Engine/Source/Runtime/Advertising/IOS/IOSAdvertising/Public/IOSAdvertising.h b/Engine/Source/Runtime/Advertising/IOS/IOSAdvertising/Public/IOSAdvertising.h
--- a/Engine/Source/Runtime/Advertising/IOS/IOSAdvertising/Public/IOSAdvertising.h
+++ b/Engine/Source/Runtime/Advertising/IOS/IOSAdvertising/Public/IOSAdvertising.h
@@ -8,7 +8,8 @@
 class FIOSAdvertisingProvider : public IAdvertisingProvider
 {
 public:
-	virtual void ShowAdBanner( bool bShowOnBottomOfScreen ) override;
+	virtual void ShowAdBanner( bool bShowOnBottomOfScreen, int32 AdID ) override;
 	virtual void HideAdBanner() override;
 	virtual void CloseAdBanner() override;
+	virtual int32 GetAdIDCount() override;
 };
diff --git a/Engine/Source/Runtime/Android/AndroidRuntimeSettings/Classes/AndroidRuntimeSettings.h b/Engine/Source/Runtime/Android/AndroidRuntimeSettings/Classes/AndroidRuntimeSettings.h
--- a/Engine/Source/Runtime/Android/AndroidRuntimeSettings/Classes/AndroidRuntimeSettings.h
+++ b/Engine/Source/Runtime/Android/AndroidRuntimeSettings/Classes/AndroidRuntimeSettings.h
@@ -242,6 +242,10 @@ public:
 	UPROPERTY(GlobalConfig, EditAnywhere, Category = GooglePlayServices)
 	FString AdMobAdUnitID;
 
+	// Identifiers for ads obtained from AdMob
+	UPROPERTY(GlobalConfig, EditAnywhere, Category = GooglePlayServices)
+	TArray<FString> AdMobAdUnitIDs;
+
 	// The unique identifier for this application (needed for IAP)
 	UPROPERTY(GlobalConfig, EditAnywhere, Category = GooglePlayServices)
 	FString GooglePlayLicenseKey;
@@ -255,5 +259,8 @@ public:
 	// UObject interface
 	virtual void PostEditChangeProperty(struct FPropertyChangedEvent& PropertyChangedEvent) override;
 	// End of UObject interface
+
+	// If the config has an AdMobAdUnitID then we migrate it on load and clear the value
+	void MigrateData();
 #endif
 };
diff --git a/Engine/Source/Runtime/Android/AndroidRuntimeSettings/Private/AndroidRuntimeSettings.cpp b/Engine/Source/Runtime/Android/AndroidRuntimeSettings/Private/AndroidRuntimeSettings.cpp
--- a/Engine/Source/Runtime/Android/AndroidRuntimeSettings/Private/AndroidRuntimeSettings.cpp
+++ b/Engine/Source/Runtime/Android/AndroidRuntimeSettings/Private/AndroidRuntimeSettings.cpp
@@ -29,4 +29,14 @@ void UAndroidRuntimeSettings::PostEditChangeProperty(struct FPropertyChangedEven
 		bBuildForES2 = true;
 	}
 }
+
+
+void UAndroidRuntimeSettings::MigrateData()
+{
+	if (!AdMobAdUnitID.IsEmpty())
+	{
+		AdMobAdUnitIDs.Add(AdMobAdUnitID);
+		AdMobAdUnitID.Empty();
+	}
+}
 #endif
diff --git a/Engine/Source/Runtime/Engine/Classes/Kismet/KismetSystemLibrary.h b/Engine/Source/Runtime/Engine/Classes/Kismet/KismetSystemLibrary.h
--- a/Engine/Source/Runtime/Engine/Classes/Kismet/KismetSystemLibrary.h
+++ b/Engine/Source/Runtime/Engine/Classes/Kismet/KismetSystemLibrary.h
@@ -1366,10 +1366,17 @@ class ENGINE_API UKismetSystemLibrary : public UBlueprintFunctionLibrary
 	 * Will show an ad banner (iAd on iOS, or AdMob on Android) on the top or bottom of screen, on top of the GL view (doesn't resize the view)
 	 * (iOS and Android only)
 	 *
+	 * @param AdIDIndex The index of the ID to select for the add to show
 	 * @param bShowOnBottomOfScreen If true, the iAd will be shown at the bottom of the screen, top otherwise
 	 */
 	UFUNCTION(BlueprintCallable, Category = "Utilities|Platform")
-	static void ShowAdBanner(bool bShowOnBottomOfScreen);
+	static void ShowAdBanner(int32 AdIDIndex, bool bShowOnBottomOfScreen);
+
+	/**
+	* Retrieves the total number of Ad IDs that can be selected between
+	*/
+	UFUNCTION(BlueprintPure, Category = "Utilities|Platform")
+	static int32 GetAdIDCount();
 
 	/**
 	 * Hides the ad banner (iAd on iOS, or AdMob on Android). Will force close the ad if it's open
diff --git a/Engine/Source/Runtime/Engine/Private/KismetSystemLibrary.cpp b/Engine/Source/Runtime/Engine/Private/KismetSystemLibrary.cpp
--- a/Engine/Source/Runtime/Engine/Private/KismetSystemLibrary.cpp
+++ b/Engine/Source/Runtime/Engine/Private/KismetSystemLibrary.cpp
@@ -3035,14 +3035,25 @@ void UKismetSystemLibrary::CollectGarbage()
 	GEngine->DeferredCommands.Add(TEXT("obj gc"));
 }
 
-void UKismetSystemLibrary::ShowAdBanner(bool bShowOnBottomOfScreen)
+void UKismetSystemLibrary::ShowAdBanner(int32 AdIDNum, bool bShowOnBottomOfScreen)
 {
 	if (IAdvertisingProvider* Provider = FAdvertising::Get().GetDefaultProvider())
 	{
-		Provider->ShowAdBanner(bShowOnBottomOfScreen);
+		Provider->ShowAdBanner(bShowOnBottomOfScreen, AdIDNum);
 	}
 }
 
+int32 UKismetSystemLibrary::GetAdIDCount()
+{
+	uint32 AdIDCount = 0;
+	if (IAdvertisingProvider* Provider = FAdvertising::Get().GetDefaultProvider())
+	{
+		AdIDCount = Provider->GetAdIDCount();
+	}
+
+	return AdIDCount;
+}
+
 void UKismetSystemLibrary::HideAdBanner()
 {
 	if (IAdvertisingProvider* Provider = FAdvertising::Get().GetDefaultProvider())
