Fixes for AdbMob multiple ad units
#ue4
#android

diff --git a/Engine/Source/Developer/Android/AndroidPlatformEditor/Private/AndroidPlatformEditorModule.cpp b/Engine/Source/Developer/Android/AndroidPlatformEditor/Private/AndroidPlatformEditorModule.cpp
--- a/Engine/Source/Developer/Android/AndroidPlatformEditor/Private/AndroidPlatformEditorModule.cpp
+++ b/Engine/Source/Developer/Android/AndroidPlatformEditor/Private/AndroidPlatformEditorModule.cpp
@@ -60,9 +60,6 @@ class FAndroidPlatformEditorModule
 		auto &AndroidDeviceDetection = FModuleManager::LoadModuleChecked<IAndroidDeviceDetectionModule>("AndroidDeviceDetection");
 		settings->SetDeviceDetection(AndroidDeviceDetection.GetAndroidDeviceDetection());
 		settings->UpdateTargetModulePaths();
-
-		auto platformSettings = GetMutableDefault<UAndroidRuntimeSettings>();
-		platformSettings->MigrateData();
 	}
 
 	virtual void ShutdownModule() override
diff --git a/Engine/Source/Developer/Android/AndroidPlatformEditor/Private/AndroidTargetSettingsCustomization.cpp b/Engine/Source/Developer/Android/AndroidPlatformEditor/Private/AndroidTargetSettingsCustomization.cpp
--- a/Engine/Source/Developer/Android/AndroidPlatformEditor/Private/AndroidTargetSettingsCustomization.cpp
+++ b/Engine/Source/Developer/Android/AndroidPlatformEditor/Private/AndroidTargetSettingsCustomization.cpp
@@ -186,9 +186,11 @@ void FAndroidTargetSettingsCustomization::BuildAppManifestSection(IDetailLayoutB
 
 	TSharedRef<IPropertyHandle> AdMobAdUnitIDProperty = DetailLayout.GetProperty(GET_MEMBER_NAME_CHECKED(UAndroidRuntimeSettings, AdMobAdUnitID));
 	AdMobAdUnitIDProperty->MarkHiddenByCustomization();
-	/*GooglePlayCategory.AddProperty(AdMobAdUnitIDProperty)
+
+	TSharedRef<IPropertyHandle> AdMobAdUnitIDsProperty = DetailLayout.GetProperty(GET_MEMBER_NAME_CHECKED(UAndroidRuntimeSettings, AdMobAdUnitIDs));
+	GooglePlayCategory.AddProperty(AdMobAdUnitIDsProperty)
 		.EditCondition(SetupForGooglePlayAttribute, NULL);
-		*/
+
 	TSharedRef<IPropertyHandle> GooglePlayLicenseKeyProperty = DetailLayout.GetProperty(GET_MEMBER_NAME_CHECKED(UAndroidRuntimeSettings, GooglePlayLicenseKey));
 	GooglePlayCategory.AddProperty(GooglePlayLicenseKeyProperty)
 		.EditCondition(SetupForGooglePlayAttribute, NULL);
diff --git a/Engine/Source/Runtime/Advertising/Android/AndroidAdvertising/Private/AndroidAdvertising.cpp b/Engine/Source/Runtime/Advertising/Android/AndroidAdvertising/Private/AndroidAdvertising.cpp
--- a/Engine/Source/Runtime/Advertising/Android/AndroidAdvertising/Private/AndroidAdvertising.cpp
+++ b/Engine/Source/Runtime/Advertising/Android/AndroidAdvertising/Private/AndroidAdvertising.cpp
@@ -10,17 +10,23 @@ IMPLEMENT_MODULE( FAndroidAdvertisingProvider, AndroidAdvertising );
 
 void FAndroidAdvertisingProvider::ShowAdBanner(bool bShowOnBottomOfScreen, int32 adID)
 {
-	FString AdUnitID;
-	bool found = GConfig->GetString(TEXT("/Script/AndroidRuntimeSettings.AndroidRuntimeSettings"), TEXT("AdMobAdUnitID"), AdUnitID, GEngineIni);
+	extern void AndroidThunkCpp_ShowAdBanner(const FString&, bool);
 
 	TArray<FString> AdUnitIDs;
-	int32 count = GConfig->GetArray(TEXT("/Script/AndroidRuntimeSettings.AndroidRuntimeSettings"), TEXT("AdMobAdUnitID"), AdUnitIDs, GEngineIni);
-	
-	if (count == 0 || adID >= count)
+	int32 count = GConfig->GetArray(TEXT("/Script/AndroidRuntimeSettings.AndroidRuntimeSettings"), TEXT("AdMobAdUnitIDs"), AdUnitIDs, GEngineIni);
+
+	if (count == 0)
+	{
+		// Fall back to checking old setting
+		FString AdUnitID;
+		bool found = GConfig->GetString(TEXT("/Script/AndroidRuntimeSettings.AndroidRuntimeSettings"), TEXT("AdMobAdUnitID"), AdUnitID, GEngineIni);
+		if (found && !AdUnitID.IsEmpty())
+			AndroidThunkCpp_ShowAdBanner(AdUnitID, bShowOnBottomOfScreen);
 		return;
+	}
 
-	extern void AndroidThunkCpp_ShowAdBanner(const FString&, bool);
-	AndroidThunkCpp_ShowAdBanner(AdUnitIDs[adID], bShowOnBottomOfScreen);
+	if (adID >= 0 && adID < count && !AdUnitIDs[adID].IsEmpty())
+		AndroidThunkCpp_ShowAdBanner(AdUnitIDs[adID], bShowOnBottomOfScreen);
 }
 
 void FAndroidAdvertisingProvider::HideAdBanner() 
@@ -38,7 +44,7 @@ void FAndroidAdvertisingProvider::CloseAdBanner()
 int32 FAndroidAdvertisingProvider::GetAdIDCount()
 {
 	TArray<FString> AdUnitIDs;
-	int32 count = GConfig->GetArray(TEXT("/Script/AndroidRuntimeSettings.AndroidRuntimeSettings"), TEXT("AdMobAdUnitID"), AdUnitIDs, GEngineIni);
+	int32 count = GConfig->GetArray(TEXT("/Script/AndroidRuntimeSettings.AndroidRuntimeSettings"), TEXT("AdMobAdUnitIDs"), AdUnitIDs, GEngineIni);
 
 	return count;
 }
\ No newline at end of file
diff --git a/Engine/Source/Runtime/Android/AndroidRuntimeSettings/Classes/AndroidRuntimeSettings.h b/Engine/Source/Runtime/Android/AndroidRuntimeSettings/Classes/AndroidRuntimeSettings.h
--- a/Engine/Source/Runtime/Android/AndroidRuntimeSettings/Classes/AndroidRuntimeSettings.h
+++ b/Engine/Source/Runtime/Android/AndroidRuntimeSettings/Classes/AndroidRuntimeSettings.h
@@ -258,9 +258,7 @@ public:
 #if WITH_EDITOR
 	// UObject interface
 	virtual void PostEditChangeProperty(struct FPropertyChangedEvent& PropertyChangedEvent) override;
+	virtual void PostInitProperties() override;
 	// End of UObject interface
-
-	// If the config has an AdMobAdUnitID then we migrate it on load and clear the value
-	void MigrateData();
 #endif
 };
diff --git a/Engine/Source/Runtime/Android/AndroidRuntimeSettings/Private/AndroidRuntimeSettings.cpp b/Engine/Source/Runtime/Android/AndroidRuntimeSettings/Private/AndroidRuntimeSettings.cpp
--- a/Engine/Source/Runtime/Android/AndroidRuntimeSettings/Private/AndroidRuntimeSettings.cpp
+++ b/Engine/Source/Runtime/Android/AndroidRuntimeSettings/Private/AndroidRuntimeSettings.cpp
@@ -30,13 +30,16 @@ void UAndroidRuntimeSettings::PostEditChangeProperty(struct FPropertyChangedEven
 	}
 }
 
-
-void UAndroidRuntimeSettings::MigrateData()
+void UAndroidRuntimeSettings::PostInitProperties()
 {
+	Super::PostInitProperties();
+
+	// If the config has an AdMobAdUnitID then we migrate it on load and clear the value
 	if (!AdMobAdUnitID.IsEmpty())
 	{
 		AdMobAdUnitIDs.Add(AdMobAdUnitID);
 		AdMobAdUnitID.Empty();
+		UpdateDefaultConfigFile();
 	}
 }
 #endif
diff --git a/Engine/Source/Runtime/Engine/Classes/Kismet/KismetSystemLibrary.h b/Engine/Source/Runtime/Engine/Classes/Kismet/KismetSystemLibrary.h
--- a/Engine/Source/Runtime/Engine/Classes/Kismet/KismetSystemLibrary.h
+++ b/Engine/Source/Runtime/Engine/Classes/Kismet/KismetSystemLibrary.h
@@ -1366,16 +1366,16 @@ class ENGINE_API UKismetSystemLibrary : public UBlueprintFunctionLibrary
 	 * Will show an ad banner (iAd on iOS, or AdMob on Android) on the top or bottom of screen, on top of the GL view (doesn't resize the view)
 	 * (iOS and Android only)
 	 *
-	 * @param AdIDIndex The index of the ID to select for the add to show
+	 * @param AdIdIndex The index of the ID to select for the add to show
 	 * @param bShowOnBottomOfScreen If true, the iAd will be shown at the bottom of the screen, top otherwise
 	 */
 	UFUNCTION(BlueprintCallable, Category = "Utilities|Platform")
-	static void ShowAdBanner(int32 AdIDIndex, bool bShowOnBottomOfScreen);
+	static void ShowAdBanner(int32 AdIdIndex, bool bShowOnBottomOfScreen);
 
 	/**
 	* Retrieves the total number of Ad IDs that can be selected between
 	*/
-	UFUNCTION(BlueprintPure, Category = "Utilities|Platform")
+	UFUNCTION(BlueprintPure, Category = "Utilities|Platform", meta = (DisplayName = "Get Ad ID Count"))
 	static int32 GetAdIDCount();
 
 	/**
diff --git a/Engine/Source/Runtime/Engine/Private/KismetSystemLibrary.cpp b/Engine/Source/Runtime/Engine/Private/KismetSystemLibrary.cpp
--- a/Engine/Source/Runtime/Engine/Private/KismetSystemLibrary.cpp
+++ b/Engine/Source/Runtime/Engine/Private/KismetSystemLibrary.cpp
@@ -3035,11 +3035,11 @@ void UKismetSystemLibrary::CollectGarbage()
 	GEngine->DeferredCommands.Add(TEXT("obj gc"));
 }
 
-void UKismetSystemLibrary::ShowAdBanner(int32 AdIDNum, bool bShowOnBottomOfScreen)
+void UKismetSystemLibrary::ShowAdBanner(int32 AdIdIndex, bool bShowOnBottomOfScreen)
 {
 	if (IAdvertisingProvider* Provider = FAdvertising::Get().GetDefaultProvider())
 	{
-		Provider->ShowAdBanner(bShowOnBottomOfScreen, AdIDNum);
+		Provider->ShowAdBanner(bShowOnBottomOfScreen, AdIdIndex);
 	}
 }
 
