deprecate IOS 9

diff --git a/Engine/Config/BaseEngine.ini b/Engine/Config/BaseEngine.ini
--- a/Engine/Config/BaseEngine.ini
+++ b/Engine/Config/BaseEngine.ini
@@ -1792,7 +1792,7 @@ BundleName=[PROJECT_NAME]
 BundleIdentifier=com.YourCompany.[PROJECT_NAME]
 VersionInfo=1.0
 FrameRateLock=PUFRL_30
-MinimumiOSVersion=IOS_9
+MinimumiOSVersion=IOS_10
 bSupportsIPad=True
 bSupportsIPhone=True
 AdditionalPlistData=
@@ -1805,7 +1805,7 @@ bGenerateCrashReportSymbols=false
 bAutomaticSigning=false
 UseFastIntrinsics=False
 EnableMathOptimisations=True
-MaxShaderLanguageVersion=1
+MaxShaderLanguageVersion=2
 
 ; These are the defaults for Android settings, and they need to be in the .ini since UBT reads the .ini settings, without instantiating the class
 [/Script/AndroidRuntimeSettings.AndroidRuntimeSettings]
diff --git a/Engine/Source/Developer/IOS/IOSPlatformEditor/Private/IOSTargetSettingsCustomization.cpp b/Engine/Source/Developer/IOS/IOSPlatformEditor/Private/IOSTargetSettingsCustomization.cpp
--- a/Engine/Source/Developer/IOS/IOSPlatformEditor/Private/IOSTargetSettingsCustomization.cpp
+++ b/Engine/Source/Developer/IOS/IOSPlatformEditor/Private/IOSTargetSettingsCustomization.cpp
@@ -1807,7 +1807,12 @@ void FIOSTargetSettingsCustomization::SetShaderStandard(int32 Value)
 			Message = LOCTEXT("iOSMetalShaderVersion2_0","Enabling Metal Shader Standard v2.0 increases the minimum operating system requirement for Metal from iOS 8.0/tvOS 9.0 or later to iOS/tvOS 11.0 or later.");
 			SetMinVersion((int32)EIOSVersion::IOS_11);
 		}
-		
+        else if (Value == 4 && (EIOSVersion)EnumValue < EIOSVersion::IOS_12)
+        {
+            Message = LOCTEXT("iOSMetalShaderVersion2_1","Enabling Metal Shader Standard v2.1 increases the minimum operating system requirement for Metal from iOS 8.0/tvOS 9.0 or later to iOS/tvOS 12.0 or later.");
+            SetMinVersion((int32)EIOSVersion::IOS_12);
+        }
+
 		
 		ShaderVersionWarningTextBox->SetError(Message);
 	}
diff --git a/Engine/Source/Programs/UnrealBuildTool/Platform/IOS/UEBuildIOS.cs b/Engine/Source/Programs/UnrealBuildTool/Platform/IOS/UEBuildIOS.cs
--- a/Engine/Source/Programs/UnrealBuildTool/Platform/IOS/UEBuildIOS.cs
+++ b/Engine/Source/Programs/UnrealBuildTool/Platform/IOS/UEBuildIOS.cs
@@ -297,12 +297,12 @@ namespace UnrealBuildTool
 			{
 				switch (MinimumIOSVersion)
 				{
-					case "IOS_10":
-						return "10.0";
 					case "IOS_11":
 						return "11.0";
+					case "IOS_12":
+						return "12.0";
 					default:
-						return "9.0";
+						return "10.0";
 				}
 			}
 		}
diff --git a/Engine/Source/Programs/UnrealBuildTool/Platform/IOS/UEDeployIOS.cs b/Engine/Source/Programs/UnrealBuildTool/Platform/IOS/UEDeployIOS.cs
--- a/Engine/Source/Programs/UnrealBuildTool/Platform/IOS/UEDeployIOS.cs
+++ b/Engine/Source/Programs/UnrealBuildTool/Platform/IOS/UEDeployIOS.cs
@@ -309,18 +309,19 @@ namespace UnrealBuildTool
 				switch (MinVersion)
 				{
 				case "IOS_61":
-					Log.TraceWarning("IOS 6 is no longer supported in UE4 as 4.11");
+					Log.TraceWarning("IOS 6 is no longer supported in UE4 as of 4.11");
 					MinVersion = "7.0";
 					break;
 				case "IOS_7":
-					Log.TraceWarning("IOS 7 is no longer supported in UE4 as 4.15");
+					Log.TraceWarning("IOS 7 is no longer supported in UE4 as of 4.15");
 					MinVersion = "7.0";
 					break;
 				case "IOS_8":
-					Log.TraceWarning("IOS 8 is no longer supported in UE4 as 4.18");
+					Log.TraceWarning("IOS 8 is no longer supported in UE4 as of 4.18");
 					MinVersion = "8.0";
 					break;
 				case "IOS_9":
+					Log.TraceWarning("IOS 9 is no longer supported in UE4 as of 4.21");
 					MinVersion = "9.0";
 					break;
 				case "IOS_10":
@@ -329,6 +330,9 @@ namespace UnrealBuildTool
 				case "IOS_11":
 					MinVersion = "11.0";
 					break;
+				case "IOS_12":
+					MinVersion = "12.0";
+					break;
 				}
 			}
 			else
diff --git a/Engine/Source/Runtime/Advertising/IOS/IOSAdvertising/Private/IOSAdvertising.cpp b/Engine/Source/Runtime/Advertising/IOS/IOSAdvertising/Private/IOSAdvertising.cpp
--- a/Engine/Source/Runtime/Advertising/IOS/IOSAdvertising/Private/IOSAdvertising.cpp
+++ b/Engine/Source/Runtime/Advertising/IOS/IOSAdvertising/Private/IOSAdvertising.cpp
@@ -12,6 +12,7 @@ DEFINE_LOG_CATEGORY_STATIC( LogAdvertising, Display, All );
 
 IMPLEMENT_MODULE( FIOSAdvertisingProvider, IOSAdvertising );
 
+#if __IPHONE_OS_VERSION_MIN_REQUIRED < __IPHONE_10_0
 @interface IOSAdvertising : UIResponder <ADBannerViewDelegate>
 	/** iAd banner view, if open */
 	@property(retain) ADBannerView* BannerView;
@@ -185,23 +186,41 @@ completion:^(BOOL finished)
 	[self.BannerView cancelBannerViewAction];
 }
 @end
+#endif
 
 void FIOSAdvertisingProvider::ShowAdBanner( bool bShowOnBottomOfScreen, int32 /*AdID*/ ) 
 {
+#if __IPHONE_OS_VERSION_MIN_REQUIRED < __IPHONE_10_0
 	[[IOSAdvertising GetDelegate] performSelectorOnMainThread:@selector(ShowAdBanner:) withObject:[NSNumber numberWithBool : bShowOnBottomOfScreen] waitUntilDone : NO];
+#else
+    UE_LOG(LogIOS, Warning, TEXT("iAd has been deprecated as of IOS 10 and is this advertising provider is now deprecated in UE 4.21"));
+#endif
 }
 
 void FIOSAdvertisingProvider::HideAdBanner() 
 {
+#if __IPHONE_OS_VERSION_MIN_REQUIRED < __IPHONE_10_0
 	[[IOSAdvertising GetDelegate] performSelectorOnMainThread:@selector(UserHideAdBanner) withObject:nil waitUntilDone : NO];
+#else
+    UE_LOG(LogIOS, Warning, TEXT("iAd has been deprecated as of IOS 10 and is this advertising provider is now deprecated in UE 4.21"));
+#endif
 }
 
 void FIOSAdvertisingProvider::CloseAdBanner() 
 {
+#if __IPHONE_OS_VERSION_MIN_REQUIRED < __IPHONE_10_0
 	[[IOSAdvertising GetDelegate] performSelectorOnMainThread:@selector(CloseAd) withObject:nil waitUntilDone : NO];
+#else
+    UE_LOG(LogIOS, Warning, TEXT("iAd has been deprecated as of IOS 10 and is this advertising provider is now deprecated in UE 4.21"));
+#endif
 }
 
 int32 FIOSAdvertisingProvider::GetAdIDCount()
 {
+#if __IPHONE_OS_VERSION_MIN_REQUIRED < __IPHONE_10_0
 	return 1;
+#else
+    UE_LOG(LogIOS, Warning, TEXT("iAd has been deprecated as of IOS 10 and is this advertising provider is now deprecated in UE 4.21"));
+    return 0;
+#endif
 }
diff --git a/Engine/Source/Runtime/IOS/IOSRuntimeSettings/Classes/IOSRuntimeSettings.h b/Engine/Source/Runtime/IOS/IOSRuntimeSettings/Classes/IOSRuntimeSettings.h
--- a/Engine/Source/Runtime/IOS/IOSRuntimeSettings/Classes/IOSRuntimeSettings.h
+++ b/Engine/Source/Runtime/IOS/IOSRuntimeSettings/Classes/IOSRuntimeSettings.h
@@ -40,7 +40,7 @@ UENUM()
 	IOS_8 = 8 UMETA(Hidden),
 
 	/** iOS 9 */
-	IOS_9 = 9 UMETA(DisplayName = "9.0"),
+	IOS_9 = 9 UMETA(Hidden),
 
 	/** iOS 10 */
 	IOS_10 = 10 UMETA(DisplayName = "10.0"),
diff --git a/Engine/Source/Runtime/IOS/IOSRuntimeSettings/Private/IOSRuntimeSettings.cpp b/Engine/Source/Runtime/IOS/IOSRuntimeSettings/Private/IOSRuntimeSettings.cpp
--- a/Engine/Source/Runtime/IOS/IOSRuntimeSettings/Private/IOSRuntimeSettings.cpp
+++ b/Engine/Source/Runtime/IOS/IOSRuntimeSettings/Private/IOSRuntimeSettings.cpp
@@ -19,7 +19,7 @@ UIOSRuntimeSettings::UIOSRuntimeSettings(const FObjectInitializer& ObjectInitial
     FrameRateLock = EPowerUsageFrameRateLock::PUFRL_30;
 	bSupportsIPad = true;
 	bSupportsIPhone = true;
-	MinimumiOSVersion = EIOSVersion::IOS_9;
+	MinimumiOSVersion = EIOSVersion::IOS_10;
 	EnableRemoteShaderCompile = false;
 	bGeneratedSYMFile = false;
 	bGeneratedSYMBundle = false;
@@ -140,10 +140,10 @@ void UIOSRuntimeSettings::PostInitProperties()
 		}
 	}
 
-	// switch IOS_6.1, IOS_7, and IOS_8 to IOS_9
-	if (MinimumiOSVersion < EIOSVersion::IOS_9)
+	// switch IOS_6.1, IOS_7, IOS_8, and IOS_9 to IOS_10
+	if (MinimumiOSVersion < EIOSVersion::IOS_10)
 	{
-		MinimumiOSVersion = EIOSVersion::IOS_9;
+		MinimumiOSVersion = EIOSVersion::IOS_10;
 	}
 	if (bSupportsOpenGLES2)
 	{
