disable advertising for TVOS

diff --git a/Engine/Source/Runtime/Advertising/IOS/IOSAdvertising/Private/IOSAdvertising.cpp b/Engine/Source/Runtime/Advertising/IOS/IOSAdvertising/Private/IOSAdvertising.cpp
--- a/Engine/Source/Runtime/Advertising/IOS/IOSAdvertising/Private/IOSAdvertising.cpp
+++ b/Engine/Source/Runtime/Advertising/IOS/IOSAdvertising/Private/IOSAdvertising.cpp
@@ -10,7 +10,7 @@ DEFINE_LOG_CATEGORY_STATIC( LogAdvertising, Display, All );
 
 IMPLEMENT_MODULE( FIOSAdvertisingProvider, IOSAdvertising );
 
-#if __IPHONE_OS_VERSION_MIN_REQUIRED < __IPHONE_10_0
+#if __IPHONE_OS_VERSION_MIN_REQUIRED < __IPHONE_10_0 && PLATFORM_IOS
 #import <iAd/ADBannerView.h>
 
 @interface IOSAdvertising : UIResponder <ADBannerViewDelegate>
