Additional TVOS fixed
#rnx

diff --git a/Engine/Source/Runtime/Advertising/IOS/IOSAdvertising/Private/IOSAdvertising.cpp b/Engine/Source/Runtime/Advertising/IOS/IOSAdvertising/Private/IOSAdvertising.cpp
--- a/Engine/Source/Runtime/Advertising/IOS/IOSAdvertising/Private/IOSAdvertising.cpp
+++ b/Engine/Source/Runtime/Advertising/IOS/IOSAdvertising/Private/IOSAdvertising.cpp
@@ -190,37 +190,37 @@ completion:^(BOOL finished)
 
 void FIOSAdvertisingProvider::ShowAdBanner( bool bShowOnBottomOfScreen, int32 /*AdID*/ ) 
 {
-#if __IPHONE_OS_VERSION_MIN_REQUIRED < __IPHONE_10_0
+#if __IPHONE_OS_VERSION_MIN_REQUIRED < __IPHONE_10_0 && !PLATFORM_TVOS
 	[[IOSAdvertising GetDelegate] performSelectorOnMainThread:@selector(ShowAdBanner:) withObject:[NSNumber numberWithBool : bShowOnBottomOfScreen] waitUntilDone : NO];
 #else
-    UE_LOG(LogIOS, Warning, TEXT("iAd has been deprecated as of IOS 10 and is this advertising provider is now deprecated in UE 4.21"));
+    UE_LOG(LogIOS, Warning, TEXT("iAd has been deprecated as of IOS 10 and this advertising provider is now deprecated in UE 4.21"));
 #endif
 }
 
 void FIOSAdvertisingProvider::HideAdBanner() 
 {
-#if __IPHONE_OS_VERSION_MIN_REQUIRED < __IPHONE_10_0
+#if __IPHONE_OS_VERSION_MIN_REQUIRED < __IPHONE_10_0 && !PLATFORM_TVOS
 	[[IOSAdvertising GetDelegate] performSelectorOnMainThread:@selector(UserHideAdBanner) withObject:nil waitUntilDone : NO];
 #else
-    UE_LOG(LogIOS, Warning, TEXT("iAd has been deprecated as of IOS 10 and is this advertising provider is now deprecated in UE 4.21"));
+    UE_LOG(LogIOS, Warning, TEXT("iAd has been deprecated as of IOS 10 and this advertising provider is now deprecated in UE 4.21"));
 #endif
 }
 
 void FIOSAdvertisingProvider::CloseAdBanner() 
 {
-#if __IPHONE_OS_VERSION_MIN_REQUIRED < __IPHONE_10_0
+#if __IPHONE_OS_VERSION_MIN_REQUIRED < __IPHONE_10_0 && !PLATFORM_TVOS
 	[[IOSAdvertising GetDelegate] performSelectorOnMainThread:@selector(CloseAd) withObject:nil waitUntilDone : NO];
 #else
-    UE_LOG(LogIOS, Warning, TEXT("iAd has been deprecated as of IOS 10 and is this advertising provider is now deprecated in UE 4.21"));
+    UE_LOG(LogIOS, Warning, TEXT("iAd has been deprecated as of IOS 10 and this advertising provider is now deprecated in UE 4.21"));
 #endif
 }
 
 int32 FIOSAdvertisingProvider::GetAdIDCount()
 {
-#if __IPHONE_OS_VERSION_MIN_REQUIRED < __IPHONE_10_0
+#if __IPHONE_OS_VERSION_MIN_REQUIRED < __IPHONE_10_0 && !PLATFORM_TVOS
 	return 1;
 #else
-    UE_LOG(LogIOS, Warning, TEXT("iAd has been deprecated as of IOS 10 and is this advertising provider is now deprecated in UE 4.21"));
+    UE_LOG(LogIOS, Warning, TEXT("iAd has been deprecated as of IOS 10 and this advertising provider is now deprecated in UE 4.21"));
     return 0;
 #endif
 }
