Fix up support for Google Play Services (remove CVE-[REDACTED_PHONE] issue)
Add AdMob app id and tag/ratings selections

#android

diff --git a/Engine/Plugins/Online/OnlineSubsystemGoogle/Source/OnlineSubsystemGoogle_UPL.xml b/Engine/Plugins/Online/OnlineSubsystemGoogle/Source/OnlineSubsystemGoogle_UPL.xml
--- a/Engine/Plugins/Online/OnlineSubsystemGoogle/Source/OnlineSubsystemGoogle_UPL.xml
+++ b/Engine/Plugins/Online/OnlineSubsystemGoogle/Source/OnlineSubsystemGoogle_UPL.xml
@@ -31,10 +31,10 @@
 	
     <buildGradleAdditions>
         <insert>
-            dependencies {
-            implementation('com.google.android.gms:play-services-auth:11.8.0')
-            }
-        </insert>
+			dependencies {
+			implementation('com.google.android.gms:play-services-auth:18.1.0')
+			}
+		</insert>
     </buildGradleAdditions>
 
 	<!-- optional updates applied to AndroidManifest.xml -->
diff --git a/Engine/Plugins/Runtime/Analytics/Adjust/Source/AndroidAdjust/Adjust_UPL.xml b/Engine/Plugins/Runtime/Analytics/Adjust/Source/AndroidAdjust/Adjust_UPL.xml
--- a/Engine/Plugins/Runtime/Analytics/Adjust/Source/AndroidAdjust/Adjust_UPL.xml
+++ b/Engine/Plugins/Runtime/Analytics/Adjust/Source/AndroidAdjust/Adjust_UPL.xml
@@ -47,7 +47,7 @@
 	<buildGradleAdditions>
 		<insert>
 dependencies {
-	implementation('com.google.android.gms:play-services-analytics:11.8.0')
+	implementation('com.google.android.gms:play-services-analytics:18.0.2')
 }
 		</insert>
 	</buildGradleAdditions>
diff --git a/Engine/Plugins/Runtime/GoogleCloudMessaging/Source/GoogleCloudMessaging/GoogleCloudMessaging_UPL.xml b/Engine/Plugins/Runtime/GoogleCloudMessaging/Source/GoogleCloudMessaging/GoogleCloudMessaging_UPL.xml
--- a/Engine/Plugins/Runtime/GoogleCloudMessaging/Source/GoogleCloudMessaging/GoogleCloudMessaging_UPL.xml
+++ b/Engine/Plugins/Runtime/GoogleCloudMessaging/Source/GoogleCloudMessaging/GoogleCloudMessaging_UPL.xml
@@ -16,7 +16,7 @@
 			<true>
 				<insert>
 dependencies {
-	implementation('com.google.android.gms:play-services-gcm:11.8.0')
+	implementation('com.google.android.gms:play-services-gcm:17.0.0')
 }
 				</insert>
 			</true>
diff --git a/Engine/Source/Developer/Android/AndroidPlatformEditor/Private/AndroidTargetSettingsCustomization.cpp b/Engine/Source/Developer/Android/AndroidPlatformEditor/Private/AndroidTargetSettingsCustomization.cpp
--- a/Engine/Source/Developer/Android/AndroidPlatformEditor/Private/AndroidTargetSettingsCustomization.cpp
+++ b/Engine/Source/Developer/Android/AndroidPlatformEditor/Private/AndroidTargetSettingsCustomization.cpp
@@ -268,6 +268,18 @@ void FAndroidTargetSettingsCustomization::BuildAppManifestSection(IDetailLayoutB
 	GooglePlayCategory.AddProperty(SupportAdMobProperty)
 		.EditCondition(SetupForGooglePlayAttribute, NULL);
 
+	TSharedRef<IPropertyHandle> AdMobAppIDProperty = DetailLayout.GetProperty(GET_MEMBER_NAME_CHECKED(UAndroidRuntimeSettings, AdMobAppID));
+	GooglePlayCategory.AddProperty(AdMobAppIDProperty);
+
+	TSharedRef<IPropertyHandle> TagForChildDirectedTreatmentProperty = DetailLayout.GetProperty(GET_MEMBER_NAME_CHECKED(UAndroidRuntimeSettings, TagForChildDirectedTreatment));
+	GooglePlayCategory.AddProperty(TagForChildDirectedTreatmentProperty);
+
+	TSharedRef<IPropertyHandle> TagForUnderAgeOfConsentProperty = DetailLayout.GetProperty(GET_MEMBER_NAME_CHECKED(UAndroidRuntimeSettings, TagForUnderAgeOfConsent));
+	GooglePlayCategory.AddProperty(TagForUnderAgeOfConsentProperty);
+
+	TSharedRef<IPropertyHandle> MaxAdContentRatingProperty = DetailLayout.GetProperty(GET_MEMBER_NAME_CHECKED(UAndroidRuntimeSettings, MaxAdContentRating));
+	GooglePlayCategory.AddProperty(MaxAdContentRatingProperty);
+
 	TSharedRef<IPropertyHandle> AdMobAdUnitIDProperty = DetailLayout.GetProperty(GET_MEMBER_NAME_CHECKED(UAndroidRuntimeSettings, AdMobAdUnitID));
 	AdMobAdUnitIDProperty->MarkHiddenByCustomization();
 
diff --git a/Engine/Source/Runtime/Advertising/Android/AndroidAdvertising/AndroidAdvertising_APL.xml b/Engine/Source/Runtime/Advertising/Android/AndroidAdvertising/AndroidAdvertising_APL.xml
--- a/Engine/Source/Runtime/Advertising/Android/AndroidAdvertising/AndroidAdvertising_APL.xml
+++ b/Engine/Source/Runtime/Advertising/Android/AndroidAdvertising/AndroidAdvertising_APL.xml
@@ -8,6 +8,24 @@
 		<!-- get packaging with ad support -->
 		<setBoolFromProperty result="bRequestGAID" ini="Engine" section="/Script/AndroidRuntimeSettings.AndroidRuntimeSettings" property="bRequestGAID" default="false"/>
 		<setBoolFromProperty result="bSupportAdMob" ini="Engine" section="/Script/AndroidRuntimeSettings.AndroidRuntimeSettings" property="bSupportAdMob" default="true"/>
+		<setBoolFromProperty result="AdMobAppID" ini="Engine" section="/Script/AndroidRuntimeSettings.AndroidRuntimeSettings" property="AdMobAppID" default="[NOT SET]"/>
+
+		<setStringFromProperty result="TagForChildDirectedTreatment" ini="Engine" section="/Script/AndroidRuntimeSettings.AndroidRuntimeSettings" property="TagForChildDirectedTreatment" default="TAG_FOR_CHILD_DIRECTED_TREATMENT_UNSPECIFIED"/>
+		<setStringFromProperty result="TagForUnderAgeOfConsent" ini="Engine" section="/Script/AndroidRuntimeSettings.AndroidRuntimeSettings" property="TagForUnderAgeOfConsent" default="TAG_FOR_UNDER_AGE_OF_CONSENT_UNSPECIFIED"/>
+		<setStringFromProperty result="MaxAdContentRating" ini="Engine" section="/Script/AndroidRuntimeSettings.AndroidRuntimeSettings" property="MaxAdContentRating" default="MAX_AD_CONTENT_RATING_PG"/>
+
+		<if condition="bSupportAdMob">
+			<setBoolIsEqual result="bAdMobAppIDNotSet" arg1="S$(AdMobAppID)" arg2="[NOT SET]"/>
+			<setBoolIsEqual result="bAdMobAppIDEmpty" arg1="S$(AdMobAppID)" arg2=""/>
+			<setBoolOr result="bAdMobAppIDInvalid" arg1="B$(bAdMobAppIDNotSet)" arg2="B$(bAdMobAppIDEmpty)"/>
+			<if condition="bAdMobAppIDInvalid">
+				<true>
+					<log text="AdMob App ID is not set; disabling AdMob support!"/>
+					<setBool result="bSupportAdMob" value="false"/>
+				</true>
+			</if>
+		</if>
+		
 		<if condition="bSupportAdMob">
 			<true>
 				<log text="Including AdMob support. Disable in Android project settings if you do not need to use it."/>
@@ -24,7 +42,7 @@
 			<true>
 				<insert>
 dependencies {
-	implementation('com.google.android.gms:play-services-ads:11.8.0')
+	implementation('com.google.android.gms:play-services-ads:18.1.0')
 	implementation('com.google.android.gms:play-services-ads-identifier:18.0.1')
 }
 				</insert>
@@ -38,6 +56,10 @@ dependencies {
 			<true>
 				<addPermission android:name="android.permission.INTERNET"/>
 				<addPermission android:name="android.permission.ACCESS_NETWORK_STATE"/>
+				<setElement result="AdMobMeta" value="meta-data" />
+				<addAttribute tag="$AdMobMeta" name="android:name" value="com.google.android.gms.ads.APPLICATION_ID" />
+				<addAttribute tag="$AdMobMeta" name="android:value" value="$S(AdMobAppID)" />
+				<addElement tag="application" name="AdMobMeta" />
 			</true>
 		</if>
 	</androidManifestUpdates>
@@ -48,11 +70,15 @@ dependencies {
 			<true>
 				<insert>
 import com.google.android.gms.common.GoogleApiAvailability;
+import com.google.android.gms.ads.MobileAds;
+import com.google.android.gms.ads.initialization.InitializationStatus;
+import com.google.android.gms.ads.initialization.OnInitializationCompleteListener;
 import com.google.android.gms.ads.AdRequest;
 import com.google.android.gms.ads.AdView;
 import com.google.android.gms.ads.AdSize;
 import com.google.android.gms.ads.AdListener;
 import com.google.android.gms.ads.InterstitialAd;
+import com.google.android.gms.ads.RequestConfiguration;
 import com.google.android.gms.ads.identifier.AdvertisingIdClient;
 import com.google.android.gms.ads.identifier.AdvertisingIdClient.Info;
 				</insert>
@@ -415,6 +441,31 @@ import com.google.android.gms.ads.identifier.AdvertisingIdClient.Info;
 			</false>
 		</if>
 	</gameActivityClassAdditions>
+	
+	<gameActivityOnCreateAdditions>
+		<if condition="bSupportAdMob">
+			<true>
+				<insert>
+		RequestConfiguration AdMobConfig = new RequestConfiguration.Builder()
+				</insert>
+				<insertValue value="				.setTagForChildDirectedTreatment(RequestConfiguration.$S(TagForChildDirectedTreatment))"/>
+				<insertNewline/>
+				<insertValue value="				.setTagForUnderAgeOfConsent(RequestConfiguration.$S(TagForUnderAgeOfConsent))"/>
+				<insertNewline/>
+				<insertValue value="				.setMaxAdContentRating(RequestConfiguration.$S(MaxAdContentRating))"/>
+				<insertNewline/>
+				<insert>
+				.build();
+		MobileAds.setRequestConfiguration(AdMobConfig);
+        MobileAds.initialize(this, new OnInitializationCompleteListener() {
+            @Override
+            public void onInitializationComplete(InitializationStatus initializationStatus) {
+            }
+        });
+				</insert>
+			</true>
+		</if>
+	</gameActivityOnCreateAdditions>
 
 	<gameActivityBeforeConfigRulesAppliedAdditions>
 		<if condition="bRequestGAID">
diff --git a/Engine/Source/Runtime/Android/AndroidRuntimeSettings/Classes/AndroidRuntimeSettings.h b/Engine/Source/Runtime/Android/AndroidRuntimeSettings/Classes/AndroidRuntimeSettings.h
--- a/Engine/Source/Runtime/Android/AndroidRuntimeSettings/Classes/AndroidRuntimeSettings.h
+++ b/Engine/Source/Runtime/Android/AndroidRuntimeSettings/Classes/AndroidRuntimeSettings.h
@@ -91,6 +91,43 @@ namespace EOculusMobileDevice
 	};
 }
 
+/** AdMob TagForChildDirectedTreament  */
+UENUM()
+namespace ETagForChildDirectedTreatment
+{
+	enum Type : int
+	{
+		TAG_FOR_CHILD_DIRECTED_TREATMENT_UNSPECIFIED = 0,
+		TAG_FOR_CHILD_DIRECTED_TREATMENT_TRUE = 1,
+		TAG_FOR_CHILD_DIRECTED_TREATMENT_FALSE = 2,
+	};
+}
+
+/** AdMob TagForChildDirectedTreament  */
+UENUM()
+namespace ETagForUnderAgeOfConsent
+{
+	enum Type : int
+	{
+		TAG_FOR_UNDER_AGE_OF_CONSENT_UNSPECIFIED = 0,
+		TAG_FOR_UNDER_AGE_OF_CONSENT_TRUE = 1,
+		TAG_FOR_UNDER_AGE_OF_CONSENT_FALSE = 2,
+	};
+}
+
+/** AdMob MaxAdContentRating  */
+UENUM()
+namespace EMaxAdContentRating
+{
+	enum Type : int
+	{
+		MAX_AD_CONTENT_RATING_G = 0,
+		MAX_AD_CONTENT_RATING_PG = 1,
+		MAX_AD_CONTENT_RATING_T = 2,
+		MAX_AD_CONTENT_RATING_MA = 2,
+	};
+}
+
 /**
  * Holds the game-specific achievement name and corresponding ID from Google Play services.
  */
@@ -449,6 +486,22 @@ public:
 	UPROPERTY(GlobalConfig, EditAnywhere, Category = GooglePlayServices, meta = (DisplayName = "Include AdMob support for ads"))
 	bool bSupportAdMob;
 
+	// AdMob Application ID
+	UPROPERTY(GlobalConfig, EditAnywhere, Category = GooglePlayServices, meta = (DisplayName = "AdMob App ID"), meta = (EditCondition = "bSupportAdMob"), meta = (ToolTip = "AdMob Application ID (see https://support.google.com/admob/answer/7356431)"))
+	FString AdMobAppID;
+
+	// Admob TagForChildDirectedTreatment (see https://developers.google.com/admob/android/targeting)
+	UPROPERTY(GlobalConfig, EditAnywhere, Category = GooglePlayServices, meta = (DisplayName = "AdMob Tag For Child-Directed Treatment"), meta = (EditCondition = "bSupportAdMob"), meta = (ToolTip = "Admob TagForChildDirectedTreatment (see https://developers.google.com/admob/android/targeting)"))
+	TEnumAsByte<ETagForChildDirectedTreatment::Type> TagForChildDirectedTreatment;
+
+	// Admob TagForUnderAgeOfConsent (see https://developers.google.com/admob/android/targeting)
+	UPROPERTY(GlobalConfig, EditAnywhere, Category = GooglePlayServices, meta = (DisplayName = "AdMob Tag For Under Age Of Consent"), meta = (EditCondition = "bSupportAdMob"), meta = (ToolTip = "Admob TagForUnderAgeOfConsent (see https://developers.google.com/admob/android/targeting)"))
+	TEnumAsByte<ETagForUnderAgeOfConsent::Type> TagForUnderAgeOfConsent;
+
+	// Admob MaxAdContentRating (see https://developers.google.com/admob/android/targeting)
+	UPROPERTY(GlobalConfig, EditAnywhere, Category = GooglePlayServices, meta = (DisplayName = "AdMob Max ad content rating"), meta = (EditCondition = "bSupportAdMob"), meta = (ToolTip = "Admob MaxAdContentRating (see https://developers.google.com/admob/android/targeting)"))
+	TEnumAsByte<EMaxAdContentRating::Type> MaxAdContentRating;
+
 	// The unique identifier for the ad obtained from AdMob.
 	UPROPERTY(GlobalConfig, EditAnywhere, Category = GooglePlayServices)
 	FString AdMobAdUnitID;
