Fix AdMob validation

#android

diff --git a/Engine/Source/Runtime/Advertising/Android/AndroidAdvertising/AndroidAdvertising_APL.xml b/Engine/Source/Runtime/Advertising/Android/AndroidAdvertising/AndroidAdvertising_APL.xml
--- a/Engine/Source/Runtime/Advertising/Android/AndroidAdvertising/AndroidAdvertising_APL.xml
+++ b/Engine/Source/Runtime/Advertising/Android/AndroidAdvertising/AndroidAdvertising_APL.xml
@@ -15,9 +15,9 @@
 		<setStringFromProperty result="MaxAdContentRating" ini="Engine" section="/Script/AndroidRuntimeSettings.AndroidRuntimeSettings" property="MaxAdContentRating" default="MAX_AD_CONTENT_RATING_PG"/>
 
 		<if condition="bSupportAdMob">
-			<setBoolIsEqual result="bAdMobAppIDNotSet" arg1="S$(AdMobAppID)" arg2="[NOT SET]"/>
-			<setBoolIsEqual result="bAdMobAppIDEmpty" arg1="S$(AdMobAppID)" arg2=""/>
-			<setBoolOr result="bAdMobAppIDInvalid" arg1="B$(bAdMobAppIDNotSet)" arg2="B$(bAdMobAppIDEmpty)"/>
+			<setBoolIsEqual result="bAdMobAppIDNotSet" arg1="$S(AdMobAppID)" arg2="[NOT SET]"/>
+			<setBoolIsEqual result="bAdMobAppIDEmpty" arg1="$S(AdMobAppID)" arg2=""/>
+			<setBoolOr result="bAdMobAppIDInvalid" arg1="$B(bAdMobAppIDNotSet)" arg2="$B(bAdMobAppIDEmpty)"/>
 			<if condition="bAdMobAppIDInvalid">
 				<true>
 					<log text="AdMob App ID is not set; disabling AdMob support!"/>
