Fix AdMob validation

#android

diff --git a/Engine/Source/Runtime/Advertising/Android/AndroidAdvertising/AndroidAdvertising_APL.xml b/Engine/Source/Runtime/Advertising/Android/AndroidAdvertising/AndroidAdvertising_APL.xml
--- a/Engine/Source/Runtime/Advertising/Android/AndroidAdvertising/AndroidAdvertising_APL.xml
+++ b/Engine/Source/Runtime/Advertising/Android/AndroidAdvertising/AndroidAdvertising_APL.xml
@@ -8,22 +8,24 @@
 		<!-- get packaging with ad support -->
 		<setBoolFromProperty result="bRequestGAID" ini="Engine" section="/Script/AndroidRuntimeSettings.AndroidRuntimeSettings" property="bRequestGAID" default="false"/>
 		<setBoolFromProperty result="bSupportAdMob" ini="Engine" section="/Script/AndroidRuntimeSettings.AndroidRuntimeSettings" property="bSupportAdMob" default="true"/>
-		<setBoolFromProperty result="AdMobAppID" ini="Engine" section="/Script/AndroidRuntimeSettings.AndroidRuntimeSettings" property="AdMobAppID" default="[NOT SET]"/>
+		<setStringFromProperty result="AdMobAppID" ini="Engine" section="/Script/AndroidRuntimeSettings.AndroidRuntimeSettings" property="AdMobAppID" default="[NOT SET]"/>
 
 		<setStringFromProperty result="TagForChildDirectedTreatment" ini="Engine" section="/Script/AndroidRuntimeSettings.AndroidRuntimeSettings" property="TagForChildDirectedTreatment" default="TAG_FOR_CHILD_DIRECTED_TREATMENT_UNSPECIFIED"/>
 		<setStringFromProperty result="TagForUnderAgeOfConsent" ini="Engine" section="/Script/AndroidRuntimeSettings.AndroidRuntimeSettings" property="TagForUnderAgeOfConsent" default="TAG_FOR_UNDER_AGE_OF_CONSENT_UNSPECIFIED"/>
 		<setStringFromProperty result="MaxAdContentRating" ini="Engine" section="/Script/AndroidRuntimeSettings.AndroidRuntimeSettings" property="MaxAdContentRating" default="MAX_AD_CONTENT_RATING_PG"/>
 
 		<if condition="bSupportAdMob">
-			<setBoolIsEqual result="bAdMobAppIDNotSet" arg1="$S(AdMobAppID)" arg2="[NOT SET]"/>
-			<setBoolIsEqual result="bAdMobAppIDEmpty" arg1="$S(AdMobAppID)" arg2=""/>
-			<setBoolOr result="bAdMobAppIDInvalid" arg1="$B(bAdMobAppIDNotSet)" arg2="$B(bAdMobAppIDEmpty)"/>
-			<if condition="bAdMobAppIDInvalid">
-				<true>
-					<log text="AdMob App ID is not set; disabling AdMob support!"/>
-					<setBool result="bSupportAdMob" value="false"/>
-				</true>
-			</if>
+			<true>
+				<setBoolIsEqual result="bAdMobAppIDNotSet" arg1="$S(AdMobAppID)" arg2="[NOT SET]"/>
+				<setBoolIsEqual result="bAdMobAppIDEmpty" arg1="$S(AdMobAppID)" arg2=""/>
+				<setBoolOr result="bAdMobAppIDInvalid" arg1="$B(bAdMobAppIDNotSet)" arg2="$B(bAdMobAppIDEmpty)"/>
+				<if condition="bAdMobAppIDInvalid">
+					<true>
+						<log text="AdMob App ID is not set; disabling AdMob support!"/>
+						<setBool result="bSupportAdMob" value="false"/>
+					</true>
+				</if>
+			</true>
 		</if>
 		
 		<if condition="bSupportAdMob">
